Override Mac paths for shared DDC folders in FDDCUtilsModule::StartupModule instead of custom [DerivedDataBackendGraph] section in MacEngine.ini. This way we cover all possible DDC ini sections and handle ?EpicDDC same way other platforms do.

This fixes the "Shared data cache path was not usable, will not use it." warnings when building DDC for the installed engine pak.

diff --git a/Engine/Config/Mac/MacEngine.ini b/Engine/Config/Mac/MacEngine.ini
--- a/Engine/Config/Mac/MacEngine.ini
+++ b/Engine/Config/Mac/MacEngine.ini
@@ -6,10 +6,6 @@ AudioDeviceModuleName=CoreAudio
 ; Defines a platform-specific volume headroom (in dB) for audio to provide better platform consistency with respect to volume levels.
 PlatformHeadroomDB=-6
 
-[DerivedDataBackendGraph]
-Shared=(Type=FileSystem, ReadOnly=false, Clean=false, Flush=false, DeleteUnused=true, UnusedFileAge=38, FoldersToClean=-1, Path=/Volumes/UE4DDC, EnvPathOverride=UE-SharedDataCachePath)
-AltShared=(Type=FileSystem, ReadOnly=true, Clean=false, Flush=false, DeleteUnused=true, UnusedFileAge=38, FoldersToClean=-1, Path=/Volumes/UE4DDC2, EnvPathOverride=UE-SharedDataCachePath2)
-
 [SystemSettings]
 r.setres=1280x720
 g.TimeoutForBlockOnRenderFence=300000
