Cleanup MacEngine.ini

diff --git a/Engine/Config/Mac/MacEngine.ini b/Engine/Config/Mac/MacEngine.ini
--- a/Engine/Config/Mac/MacEngine.ini
+++ b/Engine/Config/Mac/MacEngine.ini
@@ -1,7 +1,5 @@
 [Audio]
 AudioDeviceModuleName=CoreAudio
-;Uncomment below and comment out above line to enable new audio mixer
-;AudioDeviceModuleName=AudioMixerAudioUnit
 
 ; Defining below allows switching to audio mixer using -audiomixer commandline
 AudioMixerModuleName=AudioMixerAudioUnit
