MetalRT: Setting more enables to false.

diff --git a/Engine/Config/Mac/DataDrivenPlatformInfo.ini b/Engine/Config/Mac/DataDrivenPlatformInfo.ini
--- a/Engine/Config/Mac/DataDrivenPlatformInfo.ini
+++ b/Engine/Config/Mac/DataDrivenPlatformInfo.ini
@@ -53,10 +53,10 @@ bSupports4ComponentUAVReadWrite=true
 bSupportsVolumeTextureAtomics=false
 bSupportsDualSourceBlending=true
 EnablesHLSL2021ByDefault=1
-bSupportsRayTracing = true
+bSupportsRayTracing = false
 bSupportsRayTracingShaders = false
 bSupportsInlineRayTracing = false
-bSupportsRayTracingIndirectInstanceData = true
+bSupportsRayTracingIndirectInstanceData = false
 bSupportsByteBufferComputeShaders = true
 bSupportsHairStrandGeometry = false
 FriendlyName=LOCTEXT("FriendlyShaderPlatformName_MacSM5", "Mac Shader Model 5")
diff --git a/Engine/Shaders/Private/RayTracing/RayTracingCommon.ush b/Engine/Shaders/Private/RayTracing/RayTracingCommon.ush
--- a/Engine/Shaders/Private/RayTracing/RayTracingCommon.ush
+++ b/Engine/Shaders/Private/RayTracing/RayTracingCommon.ush
@@ -28,7 +28,7 @@
 #include "/Platform/Private/RayTracing.ush"
 #endif // OVERRIDE_RAYTRACING_USH
 
-#ifdef COMPILER_METAL
+#if COMPILER_METAL
 #include "/Engine/Private/RayTracing/RayTracingCommonMetal.ush"
 #endif // COMPILER_METAL
 
diff --git a/Engine/Shaders/Private/RayTracing/TraceRayInline.ush b/Engine/Shaders/Private/RayTracing/TraceRayInline.ush
--- a/Engine/Shaders/Private/RayTracing/TraceRayInline.ush
+++ b/Engine/Shaders/Private/RayTracing/TraceRayInline.ush
@@ -30,7 +30,7 @@
 #if COMPILER_VULKAN
 #include "/Engine/Private/RayTracing/TraceRayInlineVulkan.ush"
 #elif COMPILER_METAL
-#include "/Engine/Private/RayTracing/TraceRayInlineMetal.ush"
+// #include "/Engine/Private/RayTracing/TraceRayInlineMetal.ush"
 #endif // COMPILER_VULKAN
 
 #define PLATFORM_SUPPORTS_INLINE_RAY_TRACING_TRIANGLE_NORMALS 0
