Disable WaveOps on Metal SM5 renderer

diff --git a/Engine/Config/Mac/DataDrivenPlatformInfo.ini b/Engine/Config/Mac/DataDrivenPlatformInfo.ini
--- a/Engine/Config/Mac/DataDrivenPlatformInfo.ini
+++ b/Engine/Config/Mac/DataDrivenPlatformInfo.ini
@@ -35,7 +35,7 @@ bSupportsDistanceFields=true
 bSupportsGen5TemporalAA=true
 bSupportsLumenGI=true
 bSupportsSSDIndirect=true
-bSupportsWaveOperations=RuntimeDependent
+bSupportsWaveOperations=Unsupported
 MinimumWaveSize=32
 MaximumWaveSize=64
 bSupportsGPUScene=true
@@ -126,4 +126,4 @@ bSupports4ComponentUAVReadWrite=true
 bSupportsManualVertexFetch=false
 bSupportsVolumeTextureAtomics=false
 bSupportsShaderPipelines = false
-FriendlyName=LOCTEXT("FriendlyShaderPlatformName_MacMobile", "Mac Metal Mobile")
\ No newline at end of file
+FriendlyName=LOCTEXT("FriendlyShaderPlatformName_MacMobile", "Mac Metal Mobile")
