XR: disable ISR on Mac

We don't actually support ISR on Mac yet, so disable this until we get this stood up

#rnx

diff --git a/Engine/Config/Mac/DataDrivenPlatformInfo.ini b/Engine/Config/Mac/DataDrivenPlatformInfo.ini
--- a/Engine/Config/Mac/DataDrivenPlatformInfo.ini
+++ b/Engine/Config/Mac/DataDrivenPlatformInfo.ini
@@ -44,7 +44,7 @@ bSupportsDiaphragmDOF=true
 bSupportsVertexShaderLayer=true
 bSupportsDxc=true
 bIsSPIRV=true
-bSupportsInstancedStereo=true
+bSupportsInstancedStereo=false
 SupportsMultiViewport=RuntimeGuaranteed
 bSupportsPercentageCloserShadows=true
 bSupportsCapsuleShadows=true
@@ -84,7 +84,7 @@ bSupportsDiaphragmDOF=true
 bSupportsVertexShaderLayer=true
 bSupportsDxc=true
 bIsSPIRV=true
-bSupportsInstancedStereo=true
+bSupportsInstancedStereo=false
 SupportsMultiViewport=RuntimeGuaranteed
 bSupportsPercentageCloserShadows=true
 bSupportsCapsuleShadows=true
