Move settings from MacEngine.ini to BaseMacEngine.ini

diff --git a/Engine/Config/Mac/BaseMacEngine.ini b/Engine/Config/Mac/BaseMacEngine.ini
--- a/Engine/Config/Mac/BaseMacEngine.ini
+++ b/Engine/Config/Mac/BaseMacEngine.ini
@@ -1,11 +1,51 @@
 [Audio]
-AudioCaptureModuleName=AudioCaptureRtAudio
+; Defining below allows switching to audio mixer using -audiomixer commandline
+AudioMixerModuleName=AudioMixerCoreAudio
+PlatformFormat=OGG
+PlatformStreamingFormat=OGG
+
+; Defines a platform-specific volume headroom (in dB) for audio to provide better platform consistency with respect to volume levels.
+PlatformHeadroomDB=-6
+
+[DevOptions.Shaders]
+MaxShaderJobBatchSize=25
+
+[OnlineServices.EOS.Auth]
+DefaultExternalCredentialTypeStr=Steam
 
 [OnlineSubsystem]
 LocalPlatformName=MAC
 
+[PlatformCrypto]
+PlatformRequiresDataCrypto=True
+
 [PlatformPaths]
 UnrealBuildTool=Engine/Build/BatchFiles/Mac/Build.sh
 
-[OnlineServices.EOS.Auth]
-DefaultExternalCredentialTypeStr=Steam
+[TextureStreaming]
+; PoolSizeVRAMPercentage is how much percentage of GPU Dedicated VRAM should be used as a TexturePool cache for streaming textures (0 - unlimited streaming)
+PoolSizeVRAMPercentage=70
+
+[SystemSettings]
+r.setres=1280x720
+g.TimeoutForBlockOnRenderFence=300000
+framepro.ScopeMinTimeMicroseconds=10
+fx.NiagaraAllowRuntimeScalabilityChanges=1
+;Perplatform to PerQualityLevel conversion mapping for platform groups
+QualityLevelMapping="high"
+
+[/Script/GameplayDebugger.GameplayDebuggingControllerComponent]
+CategoryZeroBind=(Key=NumPadZero,bShift=False,bCtrl=False,bAlt=False,bCmd=False)
+CategoryOneBind=(Key=NumPadOne,bShift=False,bCtrl=False,bAlt=False,bCmd=False)
+CategoryTwoBind=(Key=NumPadTwo,bShift=False,bCtrl=False,bAlt=False,bCmd=False)
+CategoryThreeBind=(Key=NumPadThree,bShift=False,bCtrl=False,bAlt=False,bCmd=False)
+CategoryFourBind=(Key=NumPadFour,bShift=False,bCtrl=False,bAlt=False,bCmd=False)
+CategoryFiveBind=(Key=NumPadFive,bShift=False,bCtrl=False,bAlt=False,bCmd=False)
+CategorySixBind=(Key=NumPadSix,bShift=False,bCtrl=False,bAlt=False,bCmd=False)
+CategorySevenBind=(Key=NumPadSeven,bShift=False,bCtrl=False,bAlt=False,bCmd=False)
+CategoryEightBind=(Key=NumPadEight,bShift=False,bCtrl=False,bAlt=False,bCmd=False)
+CategoryNineBind=(Key=NumPadNine,bShift=False,bCtrl=False,bAlt=False,bCmd=False)
+CycleDetailsViewBind=(Key=Add,bShift=False,bCtrl=False,bAlt=False,bCmd=False)
+DebugCameraBind=(Key=Tab,bShift=False,bCtrl=False,bAlt=False,bCmd=False)
+OnScreenDebugMessagesBind=(Key=Tab,bShift=True,bCtrl=False,bAlt=False,bCmd=False)
+GameHUDBind=(Key=Tilde,bShift=True,bCtrl=False,bAlt=False,bCmd=False)
diff --git a/Engine/Config/Mac/MacEngine.ini b/Engine/Config/Mac/MacEngine.ini
deleted file mode 100644
