[Shaders] Move macro SUPPORTS_INDEPENDENT_SAMPLERS to DDSPI configuration.

While the mobile team wants to keep independent samples disbaled for PCD3D_ES3_1, it makes sense to move this as a platform dependent feature into the DDSPI configuration.

#rnx

diff --git a/Engine/Config/Android/DataDrivenPlatformInfo.ini b/Engine/Config/Android/DataDrivenPlatformInfo.ini
--- a/Engine/Config/Android/DataDrivenPlatformInfo.ini
+++ b/Engine/Config/Android/DataDrivenPlatformInfo.ini
@@ -68,6 +68,7 @@ bSupportsShaderPipelines = false
 bSupportsVertexShaderSRVs=false
 bSupportsUniformBufferObjects = true
 bSupportsDxc=true
+bSupportsIndependentSamplers=false
 bIsHlslcc=false
 FriendlyName=LOCTEXT("FriendlyShaderPlatformName_Android", "Android OpenGL Mobile")
 
@@ -90,6 +91,7 @@ bSupportsVertexShaderSRVs=false
 bSupportsUniformBufferObjects = true
 bSupportsDualSourceBlending = true
 bSupportsDxc=true
+bSupportsIndependentSamplers=true
 bIsHlslcc=false
 
 FriendlyName=LOCTEXT("FriendlyShaderPlatformName_AndroidVulkan", "Android Vulkan Mobile")
@@ -121,6 +123,7 @@ bSupportsDistanceFields=true
 bSupportsDualSourceBlending=false
 bSupportsVariableRateShading = false
 bSupportsDxc=true
+bSupportsIndependentSamplers=true
 bIsHlslcc=false
 bSupportsWaveOperations=RuntimeDependent
 MinimumWaveSize=4
diff --git a/Engine/Config/IOS/DataDrivenPlatformInfo.ini b/Engine/Config/IOS/DataDrivenPlatformInfo.ini
--- a/Engine/Config/IOS/DataDrivenPlatformInfo.ini
+++ b/Engine/Config/IOS/DataDrivenPlatformInfo.ini
@@ -56,6 +56,7 @@ bSupports4ComponentUAVReadWrite=true
 bSupportsManualVertexFetch=false
 bSupportsVolumeTextureAtomics=false
 bSupportsShaderPipelines = false
+bSupportsIndependentSamplers=true
 FriendlyName=LOCTEXT("FriendlyShaderPlatformName_IOSMobile", "IOS Metal Mobile")
 
 [ShaderPlatform METAL_MRT]
@@ -69,6 +70,7 @@ bSupports4ComponentUAVReadWrite=true
 bSupportsManualVertexFetch=false
 bSupportsVolumeTextureAtomics=false
 bSupportsDistanceFields=true
+bSupportsIndependentSamplers=true
 FriendlyName=LOCTEXT("FriendlyShaderPlatformName_IOSSM5", "IOS Metal SM5")
 
 [ShaderPlatform METAL_SIM]
@@ -81,4 +83,5 @@ bSupports4ComponentUAVReadWrite=true
 bSupportsManualVertexFetch=false
 bSupportsVolumeTextureAtomics=false
 bSupportsShaderPipelines = false
+bSupportsIndependentSamplers=true
 FriendlyName=LOCTEXT("FriendlyShaderPlatformName_IOSMobileSIM", "IOS Metal Simulator Mobile")
diff --git a/Engine/Config/Mac/DataDrivenPlatformInfo.ini b/Engine/Config/Mac/DataDrivenPlatformInfo.ini
--- a/Engine/Config/Mac/DataDrivenPlatformInfo.ini
+++ b/Engine/Config/Mac/DataDrivenPlatformInfo.ini
@@ -61,6 +61,7 @@ bSupportsHairStrandGeometry = true
 bSupportsNanite=false
 bSupportsUInt64ImageAtomics=false
 bSupportsNNEShaders=true
+bSupportsIndependentSamplers=true
 FriendlyName=LOCTEXT("FriendlyShaderPlatformName_MacSM5", "Mac Metal SM5")
 
 [ShaderPlatform METAL_SM6]
@@ -101,6 +102,7 @@ bSupportsHairStrandGeometry = true
 bSupportsNanite=true
 bSupportsUInt64ImageAtomics=true
 bSupportsNNEShaders=true
+bSupportsIndependentSamplers=true
 ;BindlessSupport=AllShaderTypes
 ;bSupportsMeshShadersTier0=false
 ;bSupportsMeshShadersTier1=false
@@ -120,6 +122,7 @@ bSupports4ComponentUAVReadWrite=true
 bSupportsManualVertexFetch=true
 bSupportsVolumeTextureAtomics=false
 bSupportsDualSourceBlending=true
+bSupportsIndependentSamplers=true
 FriendlyName=LOCTEXT("FriendlyShaderPlatformName_MacSM5_MRT", "Mac Metal SM5 MRT")
 
 [ShaderPlatform METAL_MACES3_1]
@@ -133,6 +136,7 @@ bSupportsManualVertexFetch=false
 bSupportsVolumeTextureAtomics=false
 bSupportsShaderPipelines = false
 bSupportsDualSourceBlending = true
+bSupportsIndependentSamplers=true
 FriendlyName=LOCTEXT("FriendlyShaderPlatformName_MacMobile", "Mac Metal Mobile")
 
 [PreviewPlatform METAL_SM5]
diff --git a/Engine/Config/TVOS/DataDrivenPlatformInfo.ini b/Engine/Config/TVOS/DataDrivenPlatformInfo.ini
--- a/Engine/Config/TVOS/DataDrivenPlatformInfo.ini
+++ b/Engine/Config/TVOS/DataDrivenPlatformInfo.ini
@@ -26,6 +26,7 @@ bSupportsIndexBufferUAVs=true
 bSupports4ComponentUAVReadWrite=true
 bSupportsManualVertexFetch=false
 bSupportsDualSourceBlending=true
+bSupportsIndependentSamplers=true
 FriendlyName=LOCTEXT("FriendlyShaderPlatformName_TVOS_SM5", "TVOS Metal SM5")
 
 [ShaderPlatform METAL_TVOS]
@@ -37,4 +38,5 @@ bTargetsTiledGPU=true
 bSupports4ComponentUAVReadWrite=true
 bSupportsManualVertexFetch=false
 bSupportsShaderPipelines = false
+bSupportsIndependentSamplers=true
 FriendlyName=LOCTEXT("FriendlyShaderPlatformName_TVOS_Mobile", "TVOS Metal Mobile")
\ No newline at end of file
diff --git a/Engine/Config/VulkanPC/DataDrivenPlatformInfo.ini b/Engine/Config/VulkanPC/DataDrivenPlatformInfo.ini
--- a/Engine/Config/VulkanPC/DataDrivenPlatformInfo.ini
+++ b/Engine/Config/VulkanPC/DataDrivenPlatformInfo.ini
@@ -136,6 +136,7 @@ BindlessSupport=Unsupported
 bSupportsNNEShaders=true
 bSupportsUniformBufferObjects = true
 bSupportsRealTypes=Unsupported
+bSupportsIndependentSamplers=true
 
 
 
@@ -254,6 +255,7 @@ BindlessSupport=AllShaderTypes
 bSupportsNNEShaders=true
 bSupportsUniformBufferObjects = true
 bSupportsRealTypes=Unsupported
+bSupportsIndependentSamplers=true
 
 SupportsBarycentricsSemantic=RuntimeGuaranteed
 
@@ -274,6 +276,7 @@ bSupportsShaderPipelines = false
 bSupportsVertexShaderSRVs = false
 bSupportsUniformBufferObjects = true
 bSupportsDualSourceBlending = true
+bSupportsIndependentSamplers=true
 
 FriendlyName=LOCTEXT("FriendlyShaderPlatformName_VulkanPC_Mobile", "PC Vulkan Mobile")
 
diff --git a/Engine/Config/Windows/DataDrivenPlatformInfo.ini b/Engine/Config/Windows/DataDrivenPlatformInfo.ini
--- a/Engine/Config/Windows/DataDrivenPlatformInfo.ini
+++ b/Engine/Config/Windows/DataDrivenPlatformInfo.ini
@@ -74,6 +74,7 @@ SupportsMultiViewport=RuntimeDependent
 FriendlyName=LOCTEXT("FriendlyShaderPlatformName_D3D_SM5", "PC D3D SM5")
 bSupportsDualSourceBlending=true
 bSupportsNNEShaders=true
+bSupportsIndependentSamplers=true
 
 [PreviewPlatform PCD3D_SM5]
 PlatformName=PCD3D_SM5
@@ -134,6 +135,7 @@ EnablesHLSL2021ByDefault=1
 SupportsMultiViewport=RuntimeGuaranteed
 bSupportsDualSourceBlending=true
 bSupportsNNEShaders=true
+bSupportsIndependentSamplers=true
 MaxSamplers=32
 SupportsBarycentricsSemantic=RuntimeDependent
 
@@ -159,6 +161,7 @@ bSupportsSceneDataCompressedTransforms=true
 bSupportsDebugViewShaders=true
 bSupportsShaderPipelines = false
 bSupportsDualSourceBlending=true
+bSupportsIndependentSamplers=true
 FriendlyName=LOCTEXT("FriendlyShaderPlatformName_D3D_Mobile", "PC D3D Mobile")
 SupportsMultiViewport=RuntimeDependent
 
@@ -180,4 +183,5 @@ bSupportsClipDistance = false
 bSupportsShaderPipelines = false
 bSupportsVertexShaderSRVs = false
 bSupportsUniformBufferObjects = true
+bSupportsIndependentSamplers=false
 FriendlyName=LOCTEXT("FriendlyShaderPlatformName_PCGL_Mobile", "PC OpenGL Mobile")
diff --git a/Engine/Shaders/Public/Platform.ush b/Engine/Shaders/Public/Platform.ush
--- a/Engine/Shaders/Public/Platform.ush
+++ b/Engine/Shaders/Public/Platform.ush
@@ -498,12 +498,11 @@
 #define MOBILE_EMULATION 0
 #endif
 
-// Whether the platform supports independent texture and samplers
+// Whether the platform supports independent texture and samplers (defined in DDSPI)
 // When enabled, different texture lookups can share samplers to allow more artist samplers in the base pass
 // Ideally this would just be enabled for all SM4 and above feature level platforms
-// @todo metal mrt: No reason this can't work with Metal, once cross compiler is fixed
 #ifndef SUPPORTS_INDEPENDENT_SAMPLERS
-#define SUPPORTS_INDEPENDENT_SAMPLERS (SM6_PROFILE || SM5_PROFILE || METAL_MRT_PROFILE || METAL_SM5_PROFILE || METAL_SM6_PROFILE || VULKAN_PROFILE_SM5 || VULKAN_PROFILE_SM6 || VULKAN_PROFILE)
+#define SUPPORTS_INDEPENDENT_SAMPLERS 0
 #endif
 
 // Whether the platform support pixel coverage on MSAA targets (SV_Coverage).
diff --git a/Engine/Source/Runtime/Engine/Private/ShaderCompiler/ShaderCompiler.cpp b/Engine/Source/Runtime/Engine/Private/ShaderCompiler/ShaderCompiler.cpp
--- a/Engine/Source/Runtime/Engine/Private/ShaderCompiler/ShaderCompiler.cpp
+++ b/Engine/Source/Runtime/Engine/Private/ShaderCompiler/ShaderCompiler.cpp
@@ -8360,6 +8360,11 @@ void GlobalBeginCompileShader(
 		SET_SHADER_DEFINE(Input.Environment, PLATFORM_MAX_SAMPLERS, MaxSamplers);
 	}
 
+	{
+		const bool bSupportsIndependentSamplers = FDataDrivenShaderPlatformInfo::GetSupportsIndependentSamplers((EShaderPlatform)Target.Platform);
+		SET_SHADER_DEFINE(Input.Environment, SUPPORTS_INDEPENDENT_SAMPLERS, bSupportsIndependentSamplers ? 1 : 0);
+	}
+
 	ITargetPlatform* TargetPlatform = GetTargetPlatformManager()->FindTargetPlatformWithSupport(TEXT("ShaderFormat"), ShaderFormatName);
 	bool bForwardShading = false;
 	{
diff --git a/Engine/Source/Runtime/RHI/Private/DataDrivenShaderPlatformInfo.cpp b/Engine/Source/Runtime/RHI/Private/DataDrivenShaderPlatformInfo.cpp
--- a/Engine/Source/Runtime/RHI/Private/DataDrivenShaderPlatformInfo.cpp
+++ b/Engine/Source/Runtime/RHI/Private/DataDrivenShaderPlatformInfo.cpp
@@ -297,6 +297,7 @@ void FGenericDataDrivenShaderPlatformInfo::ParseDataDrivenShaderInfo(const FConf
 	GET_SECTION_BOOL_HELPER(SupportsBarycentricsIntrinsics);
 	GET_SECTION_SUPPORT_HELPER(SupportsBarycentricsSemantic);
 	GET_SECTION_BOOL_HELPER(bSupportsWave64);
+	GET_SECTION_BOOL_HELPER(bSupportsIndependentSamplers);
 #undef GET_SECTION_BOOL_HELPER
 #undef GET_SECTION_INT_HELPER
 #undef GET_SECTION_SUPPORT_HELPER
diff --git a/Engine/Source/Runtime/RHI/Public/DataDrivenShaderPlatformInfo.h b/Engine/Source/Runtime/RHI/Public/DataDrivenShaderPlatformInfo.h
--- a/Engine/Source/Runtime/RHI/Public/DataDrivenShaderPlatformInfo.h
+++ b/Engine/Source/Runtime/RHI/Public/DataDrivenShaderPlatformInfo.h
@@ -123,6 +123,7 @@ class FGenericDataDrivenShaderPlatformInfo
 	uint32 SupportsBarycentricsIntrinsics : 1;
 	uint32 SupportsBarycentricsSemantic : int32(ERHIFeatureSupport::NumBits);
 	uint32 bSupportsWave64 : 1;
+	uint32 bSupportsIndependentSamplers : 1;
 
 	// NOTE: When adding fields, you must also add to ParseDataDrivenShaderInfo!
 	uint32 bContainsValidPlatformInfo : 1;
@@ -787,6 +788,12 @@ public:
 		return Infos[Platform].bSupportsWave64;
 	}
 
+	static FORCEINLINE_DEBUGGABLE const bool GetSupportsIndependentSamplers(const FStaticShaderPlatform Platform)
+	{
+		check(IsValid(Platform));
+		return Infos[Platform].bSupportsIndependentSamplers;
+	}
+
 	static FORCEINLINE_DEBUGGABLE const bool IsValid(const FStaticShaderPlatform Platform)
 	{
 		return Infos[Platform].bContainsValidPlatformInfo;
