- Putting Mac Editor target OS back to 11 until the buildfarm to be updated fully

diff --git a/Engine/Config/Mac/Mac_SDK.json b/Engine/Config/Mac/Mac_SDK.json
--- a/Engine/Config/Mac/Mac_SDK.json
+++ b/Engine/Config/Mac/Mac_SDK.json
@@ -1,8 +1,8 @@
 {
 	"ParentSDKFile": "../Apple/Apple_SDK.json",
 
-	"EditorDeploymentTarget": "13.0.0",
-	"EditorBuildTarget": "13.00",
+	"EditorDeploymentTarget": "11.7.9",
+	"EditorBuildTarget": "11.00",
 
 	"DeploymentTarget": "11.7.9",
 	"BuildTarget": "11.00",
