Reorganize memory.MemoryPressureCriticalThresholdMB for memory pressure assessment

diff --git a/Engine/Config/Linux/BaseLinuxEngine.ini b/Engine/Config/Linux/BaseLinuxEngine.ini
--- a/Engine/Config/Linux/BaseLinuxEngine.ini
+++ b/Engine/Config/Linux/BaseLinuxEngine.ini
@@ -9,3 +9,7 @@ DefaultExternalCredentialTypeStr=Steam
 
 [CrashReportClient]
 bAgreeToCrashUpload=false
+
+[ConsoleVariables]
+; necessary for LowMemoryWarning events to be fired
+memory.MemoryPressureCriticalThresholdMB=512
diff --git a/Engine/Config/LinuxArm64/BaseLinuxArm64Engine.ini b/Engine/Config/LinuxArm64/BaseLinuxArm64Engine.ini
--- a/Engine/Config/LinuxArm64/BaseLinuxArm64Engine.ini
+++ b/Engine/Config/LinuxArm64/BaseLinuxArm64Engine.ini
@@ -1,2 +1,6 @@
 [CrashReportClient]
 bAgreeToCrashUpload=false
+
+[ConsoleVariables]
+; necessary for LowMemoryWarning events to be fired
+memory.MemoryPressureCriticalThresholdMB=512
diff --git a/Engine/Config/Mac/BaseMacEngine.ini b/Engine/Config/Mac/BaseMacEngine.ini
--- a/Engine/Config/Mac/BaseMacEngine.ini
+++ b/Engine/Config/Mac/BaseMacEngine.ini
@@ -38,6 +38,9 @@ QualityLevelMapping="high"
 [ConsoleVariables]
 r.HairStrands.Visibility.MSAA.SamplePerPixel=1
 
+; necessary for LowMemoryWarning events to be fired
+memory.MemoryPressureCriticalThresholdMB=512
+
 [SF_METAL_SM6]
 ;BindlessResources=Enabled
 ;BindlessSamplers=Enabled
diff --git a/Engine/Config/Windows/BaseWindowsEngine.ini b/Engine/Config/Windows/BaseWindowsEngine.ini
--- a/Engine/Config/Windows/BaseWindowsEngine.ini
+++ b/Engine/Config/Windows/BaseWindowsEngine.ini
@@ -30,3 +30,6 @@ BindlessSamplers=Disabled
 MaxNumTlsSlots=2048
 MaxNumThreadsWithTlsSlots=512
 
+[ConsoleVariables]
+; necessary for LowMemoryWarning events to be fired
+memory.MemoryPressureCriticalThresholdMB=512
