Mac: Bumping the min runtime target to 13.0

diff --git a/Engine/Config/Mac/Mac_SDK.json b/Engine/Config/Mac/Mac_SDK.json
--- a/Engine/Config/Mac/Mac_SDK.json
+++ b/Engine/Config/Mac/Mac_SDK.json
@@ -4,9 +4,9 @@
 	"EditorDeploymentTarget": "11.7.9",
 	"EditorBuildTarget": "11.00",
 
-	"DeploymentTarget": "11.7.9",
-	"BuildTarget": "11.00",
+	"DeploymentTarget": "13.00",
+	"BuildTarget": "13.00",
 
-	// macOS Monterey 12.5 (Released 2022/07/22)
-	"MinSoftwareVersion": "12.5.0"
+	// macOS Ventura 13.0 (Released October 24, 2022)
+	"MinSoftwareVersion": "13.00"
 }
