HF: Fix broken Nanite on Mac

diff --git a/Engine/Config/Mac/BaseMacEngine.ini b/Engine/Config/Mac/BaseMacEngine.ini
--- a/Engine/Config/Mac/BaseMacEngine.ini
+++ b/Engine/Config/Mac/BaseMacEngine.ini
@@ -37,6 +37,7 @@ QualityLevelMapping="high"
 
 [ConsoleVariables]
 r.HairStrands.Visibility.MSAA.SamplePerPixel=1
+r.SceneCulling.ExplicitCellBounds=1
 
 ; necessary for LowMemoryWarning events to be fired
 memory.MemoryPressureCriticalThresholdMB=512
