Mac: Properly bumping editor min mac os to match rel5.4

diff --git a/Engine/Config/Mac/Mac_SDK.json b/Engine/Config/Mac/Mac_SDK.json
--- a/Engine/Config/Mac/Mac_SDK.json
+++ b/Engine/Config/Mac/Mac_SDK.json
@@ -1,8 +1,8 @@
 {
 	"ParentSDKFile": "../Apple/Apple_SDK.json",
 
-	"EditorDeploymentTarget": "11.7.9",
-	"EditorBuildTarget": "11.00",
+	"EditorDeploymentTarget": "13.00",
+	"EditorBuildTarget": "13.00",
 
 	"DeploymentTarget": "13.00",
 	"BuildTarget": "13.00",
