[Compute Framework] added compute framework support for mac platform

diff --git a/Engine/Config/Mac/DataDrivenPlatformInfo.ini b/Engine/Config/Mac/DataDrivenPlatformInfo.ini
--- a/Engine/Config/Mac/DataDrivenPlatformInfo.ini
+++ b/Engine/Config/Mac/DataDrivenPlatformInfo.ini
@@ -62,6 +62,7 @@ bSupportsNanite=false
 bSupportsUInt64ImageAtomics=false
 bSupportsNNEShaders=true
 bSupportsIndependentSamplers=true
+bSupportsComputeFramework = true
 FriendlyName=LOCTEXT("FriendlyShaderPlatformName_MacSM5", "Mac Metal SM5")
 
 [ShaderPlatform METAL_SM6]
@@ -103,6 +104,7 @@ bSupportsNanite=true
 bSupportsUInt64ImageAtomics=true
 bSupportsNNEShaders=true
 bSupportsIndependentSamplers=true
+bSupportsComputeFramework = true
 ;BindlessSupport=AllShaderTypes
 ;bSupportsMeshShadersTier0=false
 ;bSupportsMeshShadersTier1=false
