Remove bSupportsByteBufferComputeShaders DDPI property as its supported on all platforms
#jria none

diff --git a/Engine/Config/Android/DataDrivenPlatformInfo.ini b/Engine/Config/Android/DataDrivenPlatformInfo.ini
--- a/Engine/Config/Android/DataDrivenPlatformInfo.ini
+++ b/Engine/Config/Android/DataDrivenPlatformInfo.ini
@@ -120,7 +120,6 @@ bSupportsRayTracing = true
 bSupportsRayTracingShaders = false
 bSupportsInlineRayTracing = true
 bSupportsRayTracingIndirectInstanceData = true
-bSupportsByteBufferComputeShaders = true
 bEnableRayTracing=true
 bSupportsGPUSkinCache=true
 EnablesHLSL2021ByDefault=1
diff --git a/Engine/Config/Mac/DataDrivenPlatformInfo.ini b/Engine/Config/Mac/DataDrivenPlatformInfo.ini
--- a/Engine/Config/Mac/DataDrivenPlatformInfo.ini
+++ b/Engine/Config/Mac/DataDrivenPlatformInfo.ini
@@ -30,7 +30,6 @@ ShaderFormat=SF_METAL_SM5
 bIsPC=true
 bSupportsDebugViewShaders=true
 bSupportsAnisotropicMaterials=true
-bSupportsByteBufferComputeShaders=true
 bSupportsDistanceFields=true
 bSupportsGen5TemporalAA=true
 bSupportsLumenGI=true
@@ -55,7 +54,6 @@ bSupportsRayTracing = false
 bSupportsRayTracingShaders = false
 bSupportsInlineRayTracing = false
 bSupportsRayTracingIndirectInstanceData = false
-bSupportsByteBufferComputeShaders = true
 bSupportsHairStrandGeometry = true
 bSupportsNanite=false
 bSupportsUInt64ImageAtomics=false
@@ -71,7 +69,6 @@ ShaderFormat=SF_METAL_SM6
 bIsPC=true
 bSupportsDebugViewShaders=true
 bSupportsAnisotropicMaterials=true
-bSupportsByteBufferComputeShaders=true
 bSupportsDistanceFields=true
 bSupportsGen5TemporalAA=true
 bSupportsLumenGI=true
@@ -96,7 +93,6 @@ bSupportsRayTracing = false
 bSupportsRayTracingShaders = false
 bSupportsInlineRayTracing = false
 bSupportsRayTracingIndirectInstanceData = false
-bSupportsByteBufferComputeShaders = true
 bSupportsHairStrandGeometry = true
 bSupportsNanite=true
 bSupportsUInt64ImageAtomics=true
diff --git a/Engine/Config/VulkanPC/DataDrivenPlatformInfo.ini b/Engine/Config/VulkanPC/DataDrivenPlatformInfo.ini
--- a/Engine/Config/VulkanPC/DataDrivenPlatformInfo.ini
+++ b/Engine/Config/VulkanPC/DataDrivenPlatformInfo.ini
@@ -50,7 +50,6 @@ bSupportsRayTracing = true
 bSupportsRayTracingShaders = false
 bSupportsInlineRayTracing = true
 bSupportsRayTracingIndirectInstanceData = true
-bSupportsByteBufferComputeShaders = true
 
 bSupportsGPUScene = true
 
@@ -173,7 +172,6 @@ bSupportsRenderTargetWriteMask = false
 bSupportsRayTracing = true
 bSupportsInlineRayTracing = true
 bSupportsRayTracingIndirectInstanceData = true
-bSupportsByteBufferComputeShaders = true
 
 ; Enable full Vulkan ray tracing
 bSupportsRayTracingShaders = true
diff --git a/Engine/Config/Windows/DataDrivenPlatformInfo.ini b/Engine/Config/Windows/DataDrivenPlatformInfo.ini
--- a/Engine/Config/Windows/DataDrivenPlatformInfo.ini
+++ b/Engine/Config/Windows/DataDrivenPlatformInfo.ini
@@ -43,7 +43,6 @@ bSupportsGPUScene=true
 bSupportsNanite=false
 bSupportsSceneDataCompressedTransforms=true
 bSupportsComputeFramework=true
-bSupportsByteBufferComputeShaders=true
 bSupportsLumenGI=true
 bSupportsDistanceFields=true
 bSupportsSSDIndirect=true
@@ -97,7 +96,6 @@ bSupportsGPUScene=true
 bSupportsNanite=true
 bSupportsSceneDataCompressedTransforms=true
 bSupportsComputeFramework=true
-bSupportsByteBufferComputeShaders=true
 bSupportsLumenGI=true
 bSupportsDistanceFields=true
 bSupportsSSDIndirect=true
@@ -150,7 +148,6 @@ bIsPC=true
 bIsMobile=true
 bSupportsDistanceFields = false
 bSupportsComputeFramework=true
-bSupportsByteBufferComputeShaders = true
 bSupportsInstancedStereo=true
 bSupportsVertexShaderLayer=true
 bSupportsManualVertexFetch=false
diff --git a/Engine/Source/Runtime/RHI/Private/DataDrivenShaderPlatformInfo.cpp b/Engine/Source/Runtime/RHI/Private/DataDrivenShaderPlatformInfo.cpp
--- a/Engine/Source/Runtime/RHI/Private/DataDrivenShaderPlatformInfo.cpp
+++ b/Engine/Source/Runtime/RHI/Private/DataDrivenShaderPlatformInfo.cpp
@@ -273,7 +273,6 @@ void FGenericDataDrivenShaderPlatformInfo::ParseDataDrivenShaderInfo(const FConf
 	GET_SECTION_BOOL_HELPER(bSupportsShaderExecutionReordering);
 	GET_SECTION_BOOL_HELPER(bSupportsHighEndRayTracingEffects);
 	GET_SECTION_BOOL_HELPER(bSupportsUnrestrictedHalfFloatBuffers)
-	GET_SECTION_BOOL_HELPER(bSupportsByteBufferComputeShaders);
 	GET_SECTION_BOOL_HELPER(bSupportsGPUScene);
 	GET_SECTION_BOOL_HELPER(bSupportsPrimitiveShaders);
 	GET_SECTION_BOOL_HELPER(bSupportsUInt64ImageAtomics);
diff --git a/Engine/Source/Runtime/RHI/Public/DataDrivenShaderPlatformInfo.h b/Engine/Source/Runtime/RHI/Public/DataDrivenShaderPlatformInfo.h
--- a/Engine/Source/Runtime/RHI/Public/DataDrivenShaderPlatformInfo.h
+++ b/Engine/Source/Runtime/RHI/Public/DataDrivenShaderPlatformInfo.h
@@ -55,7 +55,6 @@ class FGenericDataDrivenShaderPlatformInfo
 	uint32 bSupportsShaderExecutionReordering : 1; // Does the platform support Shader Execution Reordering extensions?
 	uint32 bSupportsGPUScene : 1;
 	uint32 bSupportsUnrestrictedHalfFloatBuffers : 1;
-	uint32 bSupportsByteBufferComputeShaders : 1;
 	uint32 bSupportsPrimitiveShaders : 1;
 	uint32 bSupportsUInt64ImageAtomics : 1;
 	uint32 bRequiresVendorExtensionsForAtomics : 1;
@@ -445,12 +444,6 @@ public:
 		return Infos[Platform].bSupportsUnrestrictedHalfFloatBuffers;
 	}
 
-	static FORCEINLINE_DEBUGGABLE const bool GetSupportsByteBufferComputeShaders(const FStaticShaderPlatform Platform)
-	{
-		check(IsValid(Platform));
-		return Infos[Platform].bSupportsByteBufferComputeShaders;
-	}
-
 	static FORCEINLINE_DEBUGGABLE const ERHIFeatureSupport GetSupportsWaveOperations(const FStaticShaderPlatform Platform)
 	{
 		check(IsValid(Platform));
diff --git a/Engine/Source/Runtime/RenderCore/Private/UnifiedBuffer.cpp b/Engine/Source/Runtime/RenderCore/Private/UnifiedBuffer.cpp
--- a/Engine/Source/Runtime/RenderCore/Private/UnifiedBuffer.cpp
+++ b/Engine/Source/Runtime/RenderCore/Private/UnifiedBuffer.cpp
@@ -63,7 +63,7 @@ class FByteBufferShader : public FGlobalShader
 
 		if (ResourceType == EByteBufferResourceType::Uint_Buffer || ResourceType == EByteBufferResourceType::Uint4Aligned_Buffer)
 		{
-			return FDataDrivenShaderPlatformInfo::GetSupportsByteBufferComputeShaders(Parameters.Platform);
+			return true;
 		}
 		// Don't compile structured buffer size variations unless we need them
 		else if (ResourceType != EByteBufferResourceType::StructuredBuffer && static_cast<EByteBufferStructuredSize>(PermutationVector.Get<StructuredElementSizeDim>()) != EByteBufferStructuredSize::Uint4)
