Mac - Change r.GPUStatsEnabled to 0 by default in BaseMacEngine.ini
This is a hotfix to prevent the Landscape edit mode from crashing on Mac with the Landscape Template Map

diff --git a/Engine/Config/Mac/BaseMacEngine.ini b/Engine/Config/Mac/BaseMacEngine.ini
--- a/Engine/Config/Mac/BaseMacEngine.ini
+++ b/Engine/Config/Mac/BaseMacEngine.ini
@@ -35,6 +35,7 @@ QualityLevelMapping="high"
 [ConsoleVariables]
 r.HairStrands.Visibility.MSAA.SamplePerPixel=1
 r.SceneCulling.ExplicitCellBounds=1
+r.GPUStatsEnabled=0
 
 [SF_METAL_SM6]
 BindlessResources=Enabled
