macos requirement string update

#rnx

diff --git a/Engine/Config/Mac/Mac_SDK.json b/Engine/Config/Mac/Mac_SDK.json
--- a/Engine/Config/Mac/Mac_SDK.json
+++ b/Engine/Config/Mac/Mac_SDK.json
@@ -7,6 +7,6 @@
 	"DeploymentTarget": "14.00",
 	"BuildTarget": "14.00",
 
-	"//": "macOS Ventura 14.0 (Released Sept 23, 2023)",
+	"//": "macOS Sonoma 14.0 (Released Sept 26, 2023)",
 	"MinSoftwareVersion": "14.00"
 }
