HairStrands - Disable HairStrands rendering on Mac SM5

#tests local testing

diff --git a/Engine/Config/Mac/DataDrivenPlatformInfo.ini b/Engine/Config/Mac/DataDrivenPlatformInfo.ini
--- a/Engine/Config/Mac/DataDrivenPlatformInfo.ini
+++ b/Engine/Config/Mac/DataDrivenPlatformInfo.ini
@@ -54,7 +54,7 @@ bSupportsRayTracing = false
 bSupportsRayTracingShaders = false
 bSupportsInlineRayTracing = false
 bSupportsRayTracingIndirectInstanceData = false
-bSupportsHairStrandGeometry = true
+bSupportsHairStrandGeometry = false
 bSupportsNanite=false
 bSupportsUInt64ImageAtomics=false
 bSupportsNNEShaders=false
