Metal Bindless - Temporary disable optional bindless which is causing build machine to fail on bindful

#tests local cook

diff --git a/Engine/Config/Mac/DataDrivenPlatformInfo.ini b/Engine/Config/Mac/DataDrivenPlatformInfo.ini
--- a/Engine/Config/Mac/DataDrivenPlatformInfo.ini
+++ b/Engine/Config/Mac/DataDrivenPlatformInfo.ini
@@ -61,7 +61,7 @@ bSupportsUInt64ImageAtomics=false
 bSupportsNNEShaders=false
 bSupportsIndependentSamplers=true
 bSupportsComputeFramework = true
-bSupportsBindless=true
+bSupportsBindless=false
 FriendlyName=LOCTEXT("FriendlyShaderPlatformName_MacSM5", "Mac Metal SM5")
 
 [ShaderPlatform METAL_SM6]
@@ -122,7 +122,7 @@ bSupportsShaderPipelines = false
 bSupportsDualSourceBlending = true
 bSupportsIndependentSamplers=true
 bSupportsComputeFramework = true
-bSupportsBindless=true
+bSupportsBindless=false
 FriendlyName=LOCTEXT("FriendlyShaderPlatformName_MacMobile", "Mac Metal Mobile")
 
 [PreviewPlatform METAL_SM5]
