Turn on macOS WKWebView WebBrowser plugin backend on to replace macOS CEF3 WebBrowser plugin backend by default

diff --git a/Engine/Config/Mac/BaseMacEngine.ini b/Engine/Config/Mac/BaseMacEngine.ini
--- a/Engine/Config/Mac/BaseMacEngine.ini
+++ b/Engine/Config/Mac/BaseMacEngine.ini
@@ -63,4 +63,4 @@ OnScreenDebugMessagesBind=(Key=Tab,bShift=True,bCtrl=False,bAlt=False,bCmd=False
 GameHUDBind=(Key=Tilde,bShift=True,bCtrl=False,bAlt=False,bCmd=False)
 
 [/Script/BuildSettings.BuildSettings]
-bCompileCEF3=True
+bCompileCEF3=False
