Clean up Substrate settings.

* Windows/Linux/Mac will have all the same quality settings now
* Rough diffuse is now enabled by default on Windows/Linux/Mac + current gen consoles
* Revert all settings which matches default values.

diff --git a/Engine/Config/Linux/BaseLinuxEngine.ini b/Engine/Config/Linux/BaseLinuxEngine.ini
--- a/Engine/Config/Linux/BaseLinuxEngine.ini
+++ b/Engine/Config/Linux/BaseLinuxEngine.ini
@@ -19,5 +19,4 @@ r.Substrate.ShadingQuality=1
 r.Substrate.TileCoord8bits=0
 r.Substrate.Glints=1
 r.Substrate.ClosuresPerPixel=4
-r.Substrate.OpaqueMaterialRoughRefraction=0
 r.Refraction.Blur=1
diff --git a/Engine/Config/LinuxArm64/BaseLinuxArm64Engine.ini b/Engine/Config/LinuxArm64/BaseLinuxArm64Engine.ini
--- a/Engine/Config/LinuxArm64/BaseLinuxArm64Engine.ini
+++ b/Engine/Config/LinuxArm64/BaseLinuxArm64Engine.ini
@@ -9,5 +9,4 @@ r.Substrate.ShadingQuality=1
 r.Substrate.TileCoord8bits=0
 r.Substrate.Glints=1
 r.Substrate.ClosuresPerPixel=4
-r.Substrate.OpaqueMaterialRoughRefraction=0
 r.Refraction.Blur=1
diff --git a/Engine/Config/Mac/BaseMacEngine.ini b/Engine/Config/Mac/BaseMacEngine.ini
--- a/Engine/Config/Mac/BaseMacEngine.ini
+++ b/Engine/Config/Mac/BaseMacEngine.ini
@@ -37,6 +37,16 @@ r.HairStrands.Visibility.MSAA.SamplePerPixel=1
 r.SceneCulling.ExplicitCellBounds=1
 r.GPUStatsEnabled=0
 
+[/Script/Engine.RendererSettings]
+; Substrate settings for more advanced shading on PC and editor by default.
+r.Substrate.RoughDiffuse=1
+r.Substrate.SheenQuality=1
+r.Substrate.ShadingQuality=1
+r.Substrate.TileCoord8bits=0
+r.Substrate.Glints=1
+r.Substrate.ClosuresPerPixel=4
+r.Refraction.Blur=1
+
 [ShaderPlatformConfig METAL_ES3_1]
 BindlessConfiguration=Disabled
 
diff --git a/Engine/Config/Windows/BaseWindowsEngine.ini b/Engine/Config/Windows/BaseWindowsEngine.ini
--- a/Engine/Config/Windows/BaseWindowsEngine.ini
+++ b/Engine/Config/Windows/BaseWindowsEngine.ini
@@ -40,7 +40,6 @@ r.Substrate.ShadingQuality=1
 r.Substrate.TileCoord8bits=0
 r.Substrate.Glints=1
 r.Substrate.ClosuresPerPixel=4
-r.Substrate.OpaqueMaterialRoughRefraction=0
 r.Refraction.Blur=1
 
 [OnlineSubsystem]
