Switched on multi view SlateRHIRenderer support on macOS by default to integrate native views within Slate rendering contexts (e.g. in ApplePlatformWebBrowser's WKWebView)

#Synced-CL 50718145

diff --git a/Engine/Config/Mac/BaseMacEngine.ini b/Engine/Config/Mac/BaseMacEngine.ini
--- a/Engine/Config/Mac/BaseMacEngine.ini
+++ b/Engine/Config/Mac/BaseMacEngine.ini
@@ -74,4 +74,4 @@ GameHUDBind=(Key=Tilde,bShift=True,bCtrl=False,bAlt=False,bCmd=False)
 
 [/Script/BuildSettings.BuildSettings]
 bCompileCEF3=False
-bMultiViewPerSlateWindowSupport=False 
+bMultiViewPerSlateWindowSupport=True 
